Upgrade to nuget package Microsoft.Win32.Registry 5.0.0

#trivial
#rnx

diff --git a/Engine/Source/Programs/UnrealAndroidFileTool/UnrealAndroidFileTool.csproj b/Engine/Source/Programs/UnrealAndroidFileTool/UnrealAndroidFileTool.csproj
--- a/Engine/Source/Programs/UnrealAndroidFileTool/UnrealAndroidFileTool.csproj
+++ b/Engine/Source/Programs/UnrealAndroidFileTool/UnrealAndroidFileTool.csproj
@@ -72,7 +72,7 @@
 
   <ItemGroup>
     <PackageReference Include="Microsoft.CodeAnalysis.CSharp" Version="3.7.0" />
-    <PackageReference Include="Microsoft.Win32.Registry" Version="4.7.0" />
+    <PackageReference Include="Microsoft.Win32.Registry" Version="5.0.0" />
     <PackageReference Include="OpenTracing" Version="0.12.1" />
     <PackageReference Include="System.Management" Version="4.7.0" />
     <PackageReference Include="System.Security.Cryptography.Csp" Version="4.3.0" />
