Fixed compile issue with the android file tool.

#rnx

diff --git a/Engine/Source/Programs/UnrealAndroidFileTool/Program.cs b/Engine/Source/Programs/UnrealAndroidFileTool/Program.cs
--- a/Engine/Source/Programs/UnrealAndroidFileTool/Program.cs
+++ b/Engine/Source/Programs/UnrealAndroidFileTool/Program.cs
@@ -147,8 +147,8 @@ namespace UnrealAndroidFileTool
 				while (true)
 				{
 					Console.Write("shell> ");
-					string input = Console.ReadLine();
-					if (input == "exit" || input == "quit")
+					string? input = Console.ReadLine();
+					if (input == null || input == "exit" || input == "quit")
 					{
 						return ArgIndex;
 					}
