Improved startup sequence for AFS server with updated UnrealAndroidFileTool executables (added pulldir command)

#android

diff --git a/Engine/Source/Programs/AutomationTool/Android/AndroidFileClient.cs b/Engine/Source/Programs/AutomationTool/Android/AndroidFileClient.cs
--- a/Engine/Source/Programs/AutomationTool/Android/AndroidFileClient.cs
+++ b/Engine/Source/Programs/AutomationTool/Android/AndroidFileClient.cs
@@ -28,6 +28,8 @@ namespace AutomationTool
 {
 	public class AndroidFileClient
 	{
+		public delegate bool ShellOutputCallbackType(object Data, string Message);
+
 		protected static ILogger Logger => Log.Logger;
 
 		private const bool TimeBatching = false;
@@ -107,7 +109,6 @@ namespace AutomationTool
 			private const int kIdWindowSizeChange = 5;
 			private const int kIdInvalid = 255;
 
-
 			public OptimalADB(string ADBExecutable = "", int inPort = DefaultPort)
 			{
 				Executable = (ADBExecutable != "") ? ADBExecutable : Environment.ExpandEnvironmentVariables("%ANDROID_HOME%/platform-tools/adb" + (RuntimePlatform.IsWindows ? ".exe" : ""));
@@ -363,7 +364,7 @@ namespace AutomationTool
 				return "'" + source.Replace("'", "'\\''") + "'";
 			}
 
-			public string Shell(string Device, string Command)
+			public string Shell(string Device, string Command, int Timeout = 0, ShellOutputCallbackType OutputCallback = null, object OutputCallbackData = null)
 			{
 				string Result = "FAIL";
 
@@ -405,12 +406,21 @@ namespace AutomationTool
 						return Result;
 					}
 					Result = "";
+					string EncodedData = "";
 
 					if (bUseShellv2)
 					{
 						bool bExit = false;
 						while (!bExit)
 						{
+							if (Timeout > 0 && !ClientSocket.Poll(Timeout * 1000, SelectMode.SelectRead))
+							{
+								if (OutputCallback != null)
+								{
+									Result = "TIMEOUT";
+								}
+								break;
+							}
 							int bytesRecv = ClientSocket.Receive(buffer, kHeaderSize, SocketFlags.None);
 							if (bytesRecv <= 0)
 							{
@@ -429,10 +439,22 @@ namespace AutomationTool
 								switch (id)
 								{
 									case kIdStdOut:
-										Result += Encoding.UTF8.GetString(buffer, 0, bytesRecv);
+										EncodedData = Encoding.UTF8.GetString(buffer, 0, bytesRecv);
+										Result += EncodedData;
+										if (OutputCallback != null && !OutputCallback(OutputCallbackData, EncodedData))
+										{
+											bExit = true;
+											break;
+										}
 										break;
 									case kIdStdErr:
-										Result += Encoding.UTF8.GetString(buffer, 0, bytesRecv);
+										EncodedData = Encoding.UTF8.GetString(buffer, 0, bytesRecv);
+										Result += EncodedData;
+										if (OutputCallback != null && !OutputCallback(OutputCallbackData, EncodedData))
+										{
+											bExit = true;
+											break;
+										}
 										break;
 									case kIdExit:
 										// exit code is buffer[0]
@@ -447,12 +469,25 @@ namespace AutomationTool
 
 					while (true)
 					{
+						if (Timeout > 0 && !ClientSocket.Poll(Timeout * 1000, SelectMode.SelectRead))
+						{
+							if (OutputCallback != null)
+							{
+								Result = "TIMEOUT";
+							}
+							break;
+						}
 						int bytesRecv = ClientSocket.Receive(buffer);
 						if (bytesRecv <= 0)
 						{
 							break;
 						}
-						Result += Encoding.UTF8.GetString(buffer, 0, bytesRecv);
+						EncodedData = Encoding.UTF8.GetString(buffer, 0, bytesRecv);
+						Result += EncodedData;
+						if (OutputCallback != null && !OutputCallback(OutputCallbackData, EncodedData))
+						{
+							break;
+						}
 					}
 				}
 				catch (SocketException)
@@ -1311,6 +1346,88 @@ namespace AutomationTool
 			return false;
 		}
 
+		public static string Logcat(string Device, string Options, int Timeout = 0, ShellOutputCallbackType InCallback = null, object InCallbackData = null)
+		{
+			return adb.Shell(Device, "export ANDROID_LOG_TAGS=; exec logcat " + Options, Timeout, InCallback, InCallbackData);
+		}
+
+		public static string GetLastLogcatTime(string Device)
+		{
+			string LogcatResult = Logcat(Device, "-t 1 -v monotonic");
+			string Result = "0.000";
+			foreach (string Line in LogcatResult.Split('\n'))
+			{
+				string Trimmed = Line.TrimStart();
+				if (string.IsNullOrWhiteSpace(Trimmed) || !char.IsDigit(Trimmed[0]))
+				{
+					continue;
+				}
+				int SpaceIndex = Trimmed.IndexOf(' ');
+				int TabIndex = Trimmed.IndexOf('\t');
+				if (SpaceIndex > 0 && TabIndex > 0)
+				{
+					Result = Trimmed.Substring(0, Math.Min(SpaceIndex, TabIndex));
+				}
+				else if (SpaceIndex > 0)
+				{
+					Result = Trimmed.Substring(0, SpaceIndex);
+				}
+				else if (TabIndex > 0)
+				{
+					Result = Trimmed.Substring(0, TabIndex);
+				}
+			}
+			return Result;
+		}
+
+		private class WaitCallbackData
+		{
+			public string WaitForLine;
+			public string MessageBuffer;
+
+			public WaitCallbackData(string InWaitForLine)
+			{
+				WaitForLine = InWaitForLine;
+				MessageBuffer = "";
+			}
+
+			public bool CheckForLine(string InNewData)
+			{
+				// incoming data may not be a complete line so always append the new data
+				MessageBuffer += InNewData;
+
+				// check if we now have the search line
+				if (MessageBuffer.IndexOf(WaitForLine) >= 0)
+				{
+					MessageBuffer = "";
+					return true;
+				}
+
+				// drop old lines to keep the search faster and free memory
+				int LastLineIndex = MessageBuffer.LastIndexOf('\n');
+				if (LastLineIndex >= 0)
+				{
+					MessageBuffer.Substring(LastLineIndex + 1);
+				}
+
+				return false;
+			}
+		}
+
+		private static bool LogcatWaitCallback(object InData, string Message)
+		{
+			WaitCallbackData WorkData = (WaitCallbackData)InData;
+			return WorkData.CheckForLine(Message) ? false : true;
+		}
+
+		public static bool WaitForLogcat(string Device, string Options, string Line, int Timeout)
+		{
+			WaitCallbackData WorkData = new WaitCallbackData(Line);
+
+			string Result = Logcat(Device, Options, Timeout, LogcatWaitCallback, (object)WorkData);
+			return Result != "TIMEOUT";
+		}
+
 		public static List<string> GetInstalledReceivers(string Device)
 		{
 			List<string> Result = new List<string>();
@@ -1564,7 +1681,7 @@ namespace AutomationTool
 			foreach (string Activity in InstalledActivities)
 			{
 				bDidSendStops = true;
-				adb.Shell(Device, "am start -a com.epicgames.unreal.RemoteFileManager.intent.COMMAND2 -n " + Activity + "/com.epicgames.unreal.RemoteFileManagerActivity -e cmd stop");
+				adb.Shell(Device, "am start -W -a com.epicgames.unreal.RemoteFileManager.intent.COMMAND2 -n " + Activity + "/com.epicgames.unreal.RemoteFileManagerActivity -e cmd stop");
 			}
 
 			// it can take up to 2 seconds for running servers to terminate so check if any binds to port are still active
@@ -1680,7 +1797,7 @@ namespace AutomationTool
 
 			// sent start request (won't do anything if already started)
 //			string StartCommand = "am broadcast -a com.epicgames.unreal.RemoteFileManager.intent.COMMAND -n " + PackageName + "/com.epicgames.unreal.RemoteFileManagerReceiver -e cmd start -e token " + Token + " -ei port " + ServerPort;
-			string StartCommand = "am start -a com.epicgames.unreal.RemoteFileManager.intent.COMMAND2 -n " + PackageName + "/com.epicgames.unreal.RemoteFileManagerActivity -e cmd start -e token " + Token + " -ei port " + ServerPort;
+			string StartCommand = "am start -W -a com.epicgames.unreal.RemoteFileManager.intent.COMMAND2 -n " + PackageName + "/com.epicgames.unreal.RemoteFileManagerActivity -e cmd start -e token " + Token + " -ei port " + ServerPort;
 			adb.Shell(Device, StartCommand);
 
 			// see if we can check listen status
diff --git a/Engine/Source/Programs/UnrealAndroidFileTool/Program.cs b/Engine/Source/Programs/UnrealAndroidFileTool/Program.cs
--- a/Engine/Source/Programs/UnrealAndroidFileTool/Program.cs
+++ b/Engine/Source/Programs/UnrealAndroidFileTool/Program.cs
@@ -11,18 +11,20 @@ using System.Text;
 using System.Linq;
 using AutomationTool;
 using EpicGames.Core;
+using Microsoft.Extensions.Logging;
 
 namespace UnrealAndroidFileTool
 {
 	class Program
 	{
-		private static string AppVersion = "1.0.2";
+		private static string AppVersion = "1.0.3";
+		private static ILogger Logger = Log.Logger;
 
 		static void ShowHelp(string command)
 		{
 			if (command == "")
 			{
-				Log.TraceInformation("UnrealAndroidFileTool version {0}\n" +
+				Logger.LogInformation("UnrealAndroidFileTool version {0}\n" +
 					"Copyright Epic Games, Inc. All Rights Reserved.\n\n" +
 					"UnrealAndroidFileTool [-s Device] [-ip ipAddress] [-t port] [-p PackageName] [-k token] command\n\n" +
 					"Commands:\n\n" +
@@ -55,6 +57,7 @@ namespace UnrealAndroidFileTool
 					"mv src dst             Moves file from src to dst on device.\n" +
 //					"fileread src dst       Pulls src file from device to local dst.\n" +
 					"pull src dst           Pulls src file from device to local dst.\n" +
+					"pulldir src dst        Pulls src directory from device to local dst.\n" +
 //					"filewrite src dst      Pushes local src file to device dst.\n" +
 					"push [-c] src dst      Pushes local src file or directories to device dst.\n" +
 					"command [data]         Writes data to commandline file, or shows contents.\n" +
@@ -71,25 +74,25 @@ namespace UnrealAndroidFileTool
 			}
 			if (command == "ls")
 			{
-				Log.TraceInformation("-l\tlist permissions");
-				Log.TraceInformation("-s\tlist size");
-				Log.TraceInformation("-R\tlist recursive directory tree");
-				Log.TraceInformation("-f\tlist flat\n");
+				Logger.LogInformation("-l\tlist permissions");
+				Logger.LogInformation("-s\tlist size");
+				Logger.LogInformation("-R\tlist recursive directory tree");
+				Logger.LogInformation("-f\tlist flat\n");
 				return;
 			}
 			if (command == "stop-all")
 			{
-				Log.TraceInformation("-w\twait for all listen binds to terminate");
+				Logger.LogInformation("-w\twait for all listen binds to terminate");
 				return;
 			}
 			if (command == "push")
 			{
-				Log.TraceInformation("-c\tcompress files\n");
+				Logger.LogInformation("-c\tcompress files\n");
 				return;
 			}
 			if (command == "deploy")
 			{
-				Log.TraceInformation("-c\tcompress files\n");
+				Logger.LogInformation("-c\tcompress files\n");
 				return;
 			}
 		}
@@ -174,11 +177,11 @@ namespace UnrealAndroidFileTool
 				{
 					foreach (string Line in Receivers)
 					{
-						Log.TraceInformation("{0}", Line);
+						Logger.LogInformation("{0}", Line);
 					}
 					return ArgIndex;
 				}
-				Log.TraceInformation("No packages found with receiver.");
+				Logger.LogInformation("No packages found with receiver.");
 				return ArgIndex;
 			}
 			if (ArgIndex < args.Length && args[ArgIndex] == "stop-all")
@@ -205,14 +208,14 @@ namespace UnrealAndroidFileTool
 			{
 				ArgIndex++;
 				client.TerminateServer();
-				Log.TraceInformation("Terminated server");
+				Logger.LogInformation("Terminated server");
 				return -1;
 			}
 			if (ArgIndex < args.Length && args[ArgIndex] == "pwd")
 			{
 				ArgIndex++;
 				string result = client.Query("^^");
-				Log.TraceInformation("{0}", result);
+				Logger.LogInformation("{0}", result);
 				return ArgIndex;
 			}
 			if (ArgIndex < args.Length && args[ArgIndex] == "query")
@@ -221,7 +224,7 @@ namespace UnrealAndroidFileTool
 				if (ArgIndex < args.Length)
 				{
 					string result = client.Query(args[ArgIndex++]);
-					Log.TraceInformation("{0}", result);
+					Logger.LogInformation("{0}", result);
 				}
 				else
 				{
@@ -232,7 +235,7 @@ namespace UnrealAndroidFileTool
 						string[] parts = line.Split("\t");
 						if (parts.Length > 1)
 						{
-							Log.TraceInformation("{0,-20}{1}", parts[0], parts[1]);
+							Logger.LogInformation("{0,-20}{1}", parts[0], parts[1]);
 						}
 					}
 				}
@@ -244,12 +247,12 @@ namespace UnrealAndroidFileTool
 				if (ArgIndex < args.Length)
 				{
 					string result = client.GetProp(args[ArgIndex++]);
-					Log.TraceInformation("{0}", result);
+					Logger.LogInformation("{0}", result);
 				}
 				else
 				{
 					string result = client.GetProp("");
-					Log.TraceInformation("{0}", result);
+					Logger.LogInformation("{0}", result);
 				}
 				return ArgIndex;
 			}
@@ -259,13 +262,13 @@ namespace UnrealAndroidFileTool
 				if (ArgIndex < args.Length)
 				{
 					bool result = client.SetBaseDir(FixPath(client, args[ArgIndex]));
-					Log.TraceInformation("{0}", client.Query("^^"));
+					Logger.LogInformation("{0}", client.Query("^^"));
 					return ArgIndex++;
 				}
 				else
 				{
 					string result = client.Query("^^");
-					Log.TraceInformation("{0}", result);
+					Logger.LogInformation("{0}", result);
 				}
 				return ArgIndex;
 			}
@@ -275,7 +278,7 @@ namespace UnrealAndroidFileTool
 				if (ArgIndex < args.Length)
 				{
 					bool result = client.DirExists(args[ArgIndex]);
-					Log.TraceInformation("{0}", (result ? "true" : "false"));
+					Logger.LogInformation("{0}", (result ? "true" : "false"));
 					if (!bInShell && !result)
 					{
 						client.CloseConnection();
@@ -285,7 +288,7 @@ namespace UnrealAndroidFileTool
 				}
 				else
 				{
-					Log.TraceError("missing path");
+					Logger.LogError("missing path");
 					if (!bInShell)
 					{
 						client.CloseConnection();
@@ -347,7 +350,7 @@ namespace UnrealAndroidFileTool
 				{
 					result = client.DirList((bRecursive ? "R" : "") + (bSize ? "S" : "") + (bAttributes ? "A" : "") + (bOptions ? ":" : "") + path);
 				}
-				Log.TraceInformation("{0}", result);
+				Logger.LogInformation("{0}", result);
 				return ArgIndex++;
 			}
 			if (ArgIndex < args.Length && args[ArgIndex] == "dirlist")
@@ -356,12 +359,12 @@ namespace UnrealAndroidFileTool
 				if (ArgIndex < args.Length)
 				{
 					string result = client.DirList(FixPath(client, args[ArgIndex]));
-					Log.TraceInformation("{0}", result);
+					Logger.LogInformation("{0}", result);
 					return ArgIndex++;
 				}
 				else
 				{
-					Log.TraceError("missing path");
+					Logger.LogError("missing path");
 					if (!bInShell)
 					{
 						client.CloseConnection();
@@ -377,12 +380,12 @@ namespace UnrealAndroidFileTool
 				if (ArgIndex < args.Length)
 				{
 					string result = client.DirListFlat(FixPath(client, args[ArgIndex]));
-					Log.TraceInformation("{0}", result);
+					Logger.LogInformation("{0}", result);
 					return ArgIndex++;
 				}
 				else
 				{
-					Log.TraceError("missing path");
+					Logger.LogError("missing path");
 				}
 				return ArgIndex;
 			}
@@ -392,7 +395,7 @@ namespace UnrealAndroidFileTool
 				if (ArgIndex < args.Length)
 				{
 					bool result = client.DirCreate(FixPath(client, args[ArgIndex]));
-					Log.TraceInformation("{0}", (result ? "true" : "false"));
+					Logger.LogInformation("{0}", (result ? "true" : "false"));
 					if (!bInShell && !result)
 					{
 						client.CloseConnection();
@@ -403,7 +406,7 @@ namespace UnrealAndroidFileTool
 				}
 				else
 				{
-					Log.TraceError("missing path");
+					Logger.LogError("missing path");
 					if (!bInShell)
 					{
 						client.CloseConnection();
@@ -422,7 +425,7 @@ namespace UnrealAndroidFileTool
 					if (ArgIndex < args.Length)
 					{
 						bool result = client.DirDeleteRecurse(FixPath(client, args[ArgIndex]));
-						Log.TraceInformation("{0}", (result ? "true" : "false"));
+						Logger.LogInformation("{0}", (result ? "true" : "false"));
 						if (!bInShell && !result)
 						{
 							client.CloseConnection();
@@ -432,7 +435,7 @@ namespace UnrealAndroidFileTool
 					}
 					else
 					{
-						Log.TraceError("missing path");
+						Logger.LogError("missing path");
 						if (!bInShell)
 						{
 							client.CloseConnection();
@@ -443,7 +446,7 @@ namespace UnrealAndroidFileTool
 				else if (ArgIndex < args.Length)
 				{
 					bool result = client.DirDelete(FixPath(client, args[ArgIndex]));
-					Log.TraceInformation("{0}", (result ? "true" : "false"));
+					Logger.LogInformation("{0}", (result ? "true" : "false"));
 					if (!bInShell && !result)
 					{
 						client.CloseConnection();
@@ -453,7 +456,7 @@ namespace UnrealAndroidFileTool
 				}
 				else
 				{
-					Log.TraceError("missing path");
+					Logger.LogError("missing path");
 					if (!bInShell)
 					{
 						client.CloseConnection();
@@ -469,7 +472,7 @@ namespace UnrealAndroidFileTool
 				if (ArgIndex < args.Length)
 				{
 					bool result = client.DirDeleteRecurse(args[ArgIndex]);
-					Log.TraceInformation("{0}", (result ? "true" : "false"));
+					Logger.LogInformation("{0}", (result ? "true" : "false"));
 					if (!bInShell && !result)
 					{
 						client.CloseConnection();
@@ -479,7 +482,7 @@ namespace UnrealAndroidFileTool
 				}
 				else
 				{
-					Log.TraceError("missing path");
+					Logger.LogError("missing path");
 					if (!bInShell)
 					{
 						client.CloseConnection();
@@ -494,7 +497,7 @@ namespace UnrealAndroidFileTool
 				if (ArgIndex < args.Length)
 				{
 					bool result = client.FileExists(FixPath(client, args[ArgIndex]));
-					Log.TraceInformation("{0}", (result ? "true" : "false"));
+					Logger.LogInformation("{0}", (result ? "true" : "false"));
 					if (!bInShell && !result)
 					{
 						client.CloseConnection();
@@ -504,7 +507,7 @@ namespace UnrealAndroidFileTool
 				}
 				else
 				{
-					Log.TraceError("missing path");
+					Logger.LogError("missing path");
 					if (!bInShell)
 					{
 						client.CloseConnection();
@@ -519,7 +522,7 @@ namespace UnrealAndroidFileTool
 				if (ArgIndex < args.Length)
 				{
 					bool result = client.FileDelete(FixPath(client, args[ArgIndex]));
-					Log.TraceInformation("{0}", (result ? "true" : "false"));
+					Logger.LogInformation("{0}", (result ? "true" : "false"));
 					if (!bInShell && !result)
 					{
 						client.CloseConnection();
@@ -529,7 +532,7 @@ namespace UnrealAndroidFileTool
 				}
 				else
 				{
-					Log.TraceError("missing path");
+					Logger.LogError("missing path");
 					if (!bInShell)
 					{
 						client.CloseConnection();
@@ -544,7 +547,7 @@ namespace UnrealAndroidFileTool
 				if (ArgIndex + 1 < args.Length)
 				{
 					bool result = client.FileCopy(FixPath(client, args[ArgIndex]), FixPath(client, args[ArgIndex+1]));
-					Log.TraceInformation("{0}", (result ? "true" : "false"));
+					Logger.LogInformation("{0}", (result ? "true" : "false"));
 					if (!bInShell && !result)
 					{
 						client.CloseConnection();
@@ -554,7 +557,7 @@ namespace UnrealAndroidFileTool
 				}
 				else
 				{
-					Log.TraceError("missing paths");
+					Logger.LogError("missing paths");
 					if (!bInShell)
 					{
 						client.CloseConnection();
@@ -569,7 +572,7 @@ namespace UnrealAndroidFileTool
 				if (ArgIndex + 1 < args.Length)
 				{
 					bool result = client.FileMove(FixPath(client, args[ArgIndex]), FixPath(client, args[ArgIndex + 1]));
-					Log.TraceInformation("{0}", (result ? "true" : "false"));
+					Logger.LogInformation("{0}", (result ? "true" : "false"));
 					if (!bInShell && !result)
 					{
 						client.CloseConnection();
@@ -580,7 +583,7 @@ namespace UnrealAndroidFileTool
 				}
 				else
 				{
-					Log.TraceError("missing paths");
+					Logger.LogError("missing paths");
 					if (!bInShell)
 					{
 						client.CloseConnection();
@@ -596,7 +599,7 @@ namespace UnrealAndroidFileTool
 				if (ArgIndex + 1 < args.Length)
 				{
 					bool result = client.FileRead(args[ArgIndex], args[ArgIndex + 1]);
-					Log.TraceInformation("{0}", (result ? "true" : "false"));
+					Logger.LogInformation("{0}", (result ? "true" : "false"));
 					if (!bInShell && !result)
 					{
 						client.CloseConnection();
@@ -606,7 +609,90 @@ namespace UnrealAndroidFileTool
 				}
 				else
 				{
-					Log.TraceError("missing paths");
+					Logger.LogError("missing paths");
+					if (!bInShell)
+					{
+						client.CloseConnection();
+						System.Environment.Exit(1);
+
+					}
+				}
+				return ArgIndex;
+			}
+			if (ArgIndex < args.Length && args[ArgIndex] == "pulldir")
+			{
+				ArgIndex++;
+				if (ArgIndex + 1 < args.Length)
+				{
+					string result = client.DirListFlat(FixPath(client, args[ArgIndex]));
+					if (result == "(null)")
+					{
+						Logger.LogError("false");
+						if (!bInShell)
+						{
+							client.CloseConnection();
+							System.Environment.Exit(1);
+						}
+						return ArgIndex += 2;
+					}
+					// first line is the src path base
+					string[] Lines = result.Split('\n');
+					string BaseSrcPath = Lines[0].Substring(0, Lines[0].Length - 1);
+					int BaseSrcPathLength = BaseSrcPath.Length;
+					string DestBasePath = "";
+
+					foreach (string Line in Lines)
+					{
+						if (Line == "")
+						{
+							continue;
+						}
+						if (Line.EndsWith(":"))
+						{
+							DestBasePath = Line.Substring(BaseSrcPathLength, Line.Length - BaseSrcPathLength - 1);
+							if (DestBasePath.StartsWith("/"))
+							{
+								DestBasePath = DestBasePath.Substring(1);
+							}
+							// make directories if don't exist
+							string TargetDir = args[ArgIndex + 1] + (DestBasePath.Length > 0 ? "/" + DestBasePath : "");
+							try
+							{
+								System.IO.Directory.CreateDirectory(TargetDir);
+							}
+							catch (Exception)
+							{
+								Logger.LogError("Failed to create {0}\nfalse", TargetDir);
+								if (!bInShell)
+								{
+									client.CloseConnection();
+									System.Environment.Exit(1);
+								}
+								return ArgIndex += 2;
+							}
+							continue;
+						}
+						string Src = BaseSrcPath + "/" + (DestBasePath.Length > 0 ? DestBasePath + "/" : "") + Line;
+						string Dest = args[ArgIndex + 1] + "/" + (DestBasePath.Length > 0 ? DestBasePath + "/" : "") + Line;
+						bool copyresult = client.FileRead(Src, Dest);
+						if (!bInShell && !copyresult)
+						{
+							Logger.LogError("false");
+							client.CloseConnection();
+							System.Environment.Exit(1);
+						}
+					}
+					Logger.LogInformation("true");
+					if (!bInShell)
+					{
+						client.CloseConnection();
+						System.Environment.Exit(1);
+					}
+					return ArgIndex += 2;
+				}
+				else
+				{
+					Logger.LogError("missing paths");
 					if (!bInShell)
 					{
 						client.CloseConnection();
@@ -621,9 +707,9 @@ namespace UnrealAndroidFileTool
 				ArgIndex++;
 				if (ArgIndex + 1 < args.Length)
 				{
-					//Log.TraceInformation("Src: '{0}', Dest: '{1}'", args[ArgIndex], args[ArgIndex + 1]);
+					//Logger.LogInformation("Src: '{0}', Dest: '{1}'", args[ArgIndex], args[ArgIndex + 1]);
 					bool result = client.FileWrite(args[ArgIndex], args[ArgIndex + 1]);
-					Log.TraceInformation("{0}", (result ? "true" : "false"));
+					Logger.LogInformation("{0}", (result ? "true" : "false"));
 					if (!bInShell && !result)
 					{
 						client.CloseConnection();
@@ -633,7 +719,7 @@ namespace UnrealAndroidFileTool
 				}
 				else
 				{
-					Log.TraceError("missing paths");
+					Logger.LogError("missing paths");
 					if (!bInShell)
 					{
 						client.CloseConnection();
@@ -674,7 +760,7 @@ namespace UnrealAndroidFileTool
 						}
 						catch (Exception)
 						{
-							Log.TraceInformation("false");
+							Logger.LogInformation("false");
 							if (!bInShell)
 							{
 								client.CloseConnection();
@@ -684,7 +770,7 @@ namespace UnrealAndroidFileTool
 							return ArgIndex;
 						}
 
-						Log.TraceInformation("{0}", (result ? "true" : "false"));
+						Logger.LogInformation("{0}", (result ? "true" : "false"));
 						if (!bInShell && !result)
 						{
 							client.CloseConnection();
@@ -694,7 +780,7 @@ namespace UnrealAndroidFileTool
 					}
 					else
 					{
-						Log.TraceError("missing paths");
+						Logger.LogError("missing paths");
 						if (!bInShell)
 						{
 							client.CloseConnection();
@@ -704,7 +790,7 @@ namespace UnrealAndroidFileTool
 				}
 				else
 				{
-					Log.TraceError("missing paths");
+					Logger.LogError("missing paths");
 					if (!bInShell)
 					{
 						client.CloseConnection();
@@ -723,7 +809,7 @@ namespace UnrealAndroidFileTool
 						command = command + " " + args[ArgIndex++];
 					}
 					bool result = client.FileWriteString(command, "^commandfile");
-					Log.TraceInformation("{0}", (result ? "true" : "false"));
+					Logger.LogInformation("{0}", (result ? "true" : "false"));
 					if (!bInShell && !result)
 					{
 						client.CloseConnection();
@@ -746,12 +832,12 @@ namespace UnrealAndroidFileTool
 						}
 						catch (Exception e)
 						{
-							Log.TraceError("Error reading temp file {0}, Exception: {1}", TempFilename, e.ToString());
+							Logger.LogError("Error reading temp file {0}, Exception: {1}", TempFilename, e.ToString());
 						}
 					}
 					else
 					{
-						Log.TraceInformation("No commandline file found.");
+						Logger.LogInformation("No commandline file found.");
 					}
 				}
 				return ArgIndex;
@@ -785,11 +871,11 @@ namespace UnrealAndroidFileTool
 					result = client.FileWriteString(command, "^commandfile");
 					if (!bInShell && !result)
 					{
-						Log.TraceError("false");
+						Logger.LogError("false");
 						client.CloseConnection();
 						System.Environment.Exit(1);
 					}
-					Log.TraceInformation("New commandline: {0}", command);
+					Logger.LogInformation("New commandline: {0}", command);
 					return ArgIndex;
 				}
 				else
@@ -807,7 +893,7 @@ namespace UnrealAndroidFileTool
 						}
 						catch (Exception e)
 						{
-							Log.TraceError("Error reading temp file {0}, Exception: {1}", TempFilename, e.ToString());
+							Logger.LogError("Error reading temp file {0}, Exception: {1}", TempFilename, e.ToString());
 							if (!bInShell)
 							{
 								client.CloseConnection();
@@ -817,7 +903,7 @@ namespace UnrealAndroidFileTool
 					}
 					else
 					{
-						Log.TraceInformation("No commandline file found.");
+						Logger.LogInformation("No commandline file found.");
 						if (!bInShell)
 						{
 							client.CloseConnection();
@@ -848,7 +934,7 @@ namespace UnrealAndroidFileTool
 						}
 						catch (Exception e)
 						{
-							Log.TraceError("Error reading temp file {0}, Exception: {1}", TempFilename, e.ToString());
+							Logger.LogError("Error reading temp file {0}, Exception: {1}", TempFilename, e.ToString());
 							if (!bInShell)
 							{
 								client.CloseConnection();
@@ -863,11 +949,11 @@ namespace UnrealAndroidFileTool
 					result = client.FileWriteString(command, "^commandfile");
 					if (!bInShell && !result)
 					{
-						Log.TraceError("false");
+						Logger.LogError("false");
 						client.CloseConnection();
 						System.Environment.Exit(1);
 					}
-					Log.TraceInformation("New commandline: {0}", command);
+					Logger.LogInformation("New commandline: {0}", command);
 					return ArgIndex;
 				}
 				else
@@ -885,7 +971,7 @@ namespace UnrealAndroidFileTool
 						}
 						catch (Exception e)
 						{
-							Log.TraceError("Error reading temp file {0}, Exception: {1}", TempFilename, e.ToString());
+							Logger.LogError("Error reading temp file {0}, Exception: {1}", TempFilename, e.ToString());
 							if (!bInShell)
 							{
 								client.CloseConnection();
@@ -895,7 +981,7 @@ namespace UnrealAndroidFileTool
 					}
 					else
 					{
-						Log.TraceInformation("No commandline file found.");
+						Logger.LogInformation("No commandline file found.");
 						if (!bInShell)
 						{
 							client.CloseConnection();
@@ -923,12 +1009,12 @@ namespace UnrealAndroidFileTool
 						}
 						catch (Exception e)
 						{
-							Log.TraceError("Error reading temp file {0}, Exception: {1}", TempFilename, e.ToString());
+							Logger.LogError("Error reading temp file {0}, Exception: {1}", TempFilename, e.ToString());
 						}
 					}
 					else
 					{
-						Log.TraceInformation("{0}", (result ? "true" : "false"));
+						Logger.LogInformation("{0}", (result ? "true" : "false"));
 						if (!bInShell)
 						{
 							client.CloseConnection();
@@ -939,7 +1025,7 @@ namespace UnrealAndroidFileTool
 				}
 				else
 				{
-					Log.TraceError("missing path");
+					Logger.LogError("missing path");
 					if (!bInShell)
 					{
 						client.CloseConnection();
@@ -986,7 +1072,7 @@ namespace UnrealAndroidFileTool
 						}
 						catch (Exception e)
 						{
-							Log.TraceError("Error reading temp file {0}, Exception: {1}", args[ArgIndex], e.ToString());
+							Logger.LogError("Error reading temp file {0}, Exception: {1}", args[ArgIndex], e.ToString());
 							if (!bInShell)
 							{
 								client.CloseConnection();
@@ -996,7 +1082,7 @@ namespace UnrealAndroidFileTool
 					}
 					else
 					{
-						Log.TraceError("missing path");
+						Logger.LogError("missing path");
 						if (!bInShell)
 						{
 							client.CloseConnection();
@@ -1006,7 +1092,7 @@ namespace UnrealAndroidFileTool
 				}
 				else
 				{
-					Log.TraceError("missing path");
+					Logger.LogError("missing path");
 					if (!bInShell)
 					{
 						client.CloseConnection();
@@ -1016,7 +1102,7 @@ namespace UnrealAndroidFileTool
 				return ArgIndex;
 			}
 
-			Log.TraceWarning("Unknown command");
+			Logger.LogWarning("Unknown command");
 			return ArgIndex;
 		}
 
@@ -1029,15 +1115,15 @@ namespace UnrealAndroidFileTool
 				{
 					if (PackageName == "")
 					{
-						Log.TraceError("Need package name to start server on {0}!", Device);
+						Logger.LogError("Need package name to start server on {0}!", Device);
 						return null;
 					}
 					else
 					{
-						Log.TraceInformation("Trying to start file server {0}", PackageName);
+						Logger.LogInformation("Trying to start file server {0}", PackageName);
 						if (!client.StartServer(PackageName, Token, IPAddress))
 						{
-							Log.TraceError("Unable to connect to " + Device);
+							Logger.LogError("Unable to connect to " + Device);
 							return null;
 						}
 					}
@@ -1049,18 +1135,18 @@ namespace UnrealAndroidFileTool
 				{
 					if (PackageName == "")
 					{
-						Log.TraceError("Need package name to start server on {0}!", Device);
+						Logger.LogError("Need package name to start server on {0}!", Device);
 						client.CloseConnection();
 						return null;
 					}
 
-					Log.TraceInformation("Connected to wrong server {0}, trying again", DevicePackageName);
+					Logger.LogInformation("Connected to wrong server {0}, trying again", DevicePackageName);
 					client.TerminateServer();
 
-					Log.TraceInformation("Trying to start file server {0}", PackageName);
+					Logger.LogInformation("Trying to start file server {0}", PackageName);
 					if (!client.StartServer(PackageName))
 					{
-						Log.TraceError("Failed to start server");
+						Logger.LogError("Failed to start server");
 						return null;
 					}
 				}
@@ -1111,7 +1197,7 @@ namespace UnrealAndroidFileTool
 				}
 				else
 				{
-					Log.TraceError("Missing device id");
+					Logger.LogError("Missing device id");
 					System.Environment.Exit(1);
 				}
 			}
@@ -1126,7 +1212,7 @@ namespace UnrealAndroidFileTool
 				}
 				else
 				{
-					Log.TraceError("Missing package name");
+					Logger.LogError("Missing package name");
 					System.Environment.Exit(1);
 				}
 			}
@@ -1141,7 +1227,7 @@ namespace UnrealAndroidFileTool
 				}
 				else
 				{
-					Log.TraceError("Missing security token key");
+					Logger.LogError("Missing security token key");
 					System.Environment.Exit(1);
 				}
 			}
@@ -1158,13 +1244,13 @@ namespace UnrealAndroidFileTool
 					}
 					catch (Exception)
 					{
-						Log.TraceError("Invalid port");
+						Logger.LogError("Invalid port");
 						System.Environment.Exit(1);
 					}
 				}
 				else
 				{
-					Log.TraceError("Missing port");
+					Logger.LogError("Missing port");
 					System.Environment.Exit(1);
 				}
 			}
@@ -1179,7 +1265,7 @@ namespace UnrealAndroidFileTool
 				}
 				else
 				{
-					Log.TraceError("Missing IP address");
+					Logger.LogError("Missing IP address");
 					System.Environment.Exit(1);
 				}
 			}
@@ -1190,7 +1276,7 @@ namespace UnrealAndroidFileTool
 				AndroidFileClient.GetConnectedDevices(out devices);
 				foreach (string DeviceName in devices)
 				{
-					Log.TraceInformation("{0}", DeviceName);
+					Logger.LogInformation("{0}", DeviceName);
 				}
 				return;
 			}
@@ -1202,7 +1288,7 @@ namespace UnrealAndroidFileTool
 				AndroidFileClient.GetConnectedDevices(out devices);
 				if (devices.Count == 0)
 				{
-					Log.TraceError("No connected devices");
+					Logger.LogError("No connected devices");
 				}
 				else if (devices.Count == 1)
 				{
@@ -1210,17 +1296,17 @@ namespace UnrealAndroidFileTool
 				}
 				else
 				{
-					Log.TraceError("Multiple devices connected, select one -s!");
+					Logger.LogError("Multiple devices connected, select one -s!");
 					foreach (string DeviceName in devices)
 					{
-						Log.TraceInformation("{0}", DeviceName);
+						Logger.LogInformation("{0}", DeviceName);
 					}
 				}
 			}
 
 			if (Device == "")
 			{
-				Log.TraceError("No devices attached!");
+				Logger.LogError("No devices attached!");
 				System.Environment.Exit(0);
 			}
 
@@ -1233,11 +1319,11 @@ namespace UnrealAndroidFileTool
 				{
 					foreach (string Line in Receivers)
 					{
-						Log.TraceInformation("{0}", Line);
+						Logger.LogInformation("{0}", Line);
 					}
 					System.Environment.Exit(0);
 				}
-				Log.TraceInformation("No packages found with receiver.");
+				Logger.LogInformation("No packages found with receiver.");
 				System.Environment.Exit(0);
 			}
 			if (ArgIndex < args.Length && args[ArgIndex] == "stop-all")
@@ -1259,7 +1345,7 @@ namespace UnrealAndroidFileTool
 			{
 				System.Environment.Exit(1);
 			}
-			Log.TraceInformation("Connected!");
+			Logger.LogInformation("Connected!");
 
 			if (ArgIndex < args.Length)
 			{
