Fix issue with UnrealAndroidFileTool packages command

#android

diff --git a/Engine/Source/Programs/UnrealAndroidFileTool/Program.cs b/Engine/Source/Programs/UnrealAndroidFileTool/Program.cs
--- a/Engine/Source/Programs/UnrealAndroidFileTool/Program.cs
+++ b/Engine/Source/Programs/UnrealAndroidFileTool/Program.cs
@@ -17,7 +17,7 @@ namespace UnrealAndroidFileTool
 {
 	class Program
 	{
-		private static string AppVersion = "1.0.3";
+		private static string AppVersion = "1.0.4";
 		private static ILogger Logger = Log.Logger;
 
 		static void ShowHelp(string command)
@@ -173,6 +173,7 @@ namespace UnrealAndroidFileTool
 			{
 				ArgIndex++;
 				List<string> Receivers = AndroidFileClient.GetInstalledReceivers(client.GetDevice());
+				AndroidFileClient.GetInstalledActivities(client.GetDevice()).ForEach(item => { if (!Receivers.Contains(item)) Receivers.Add(item); });
 				if (Receivers.Count > 0)
 				{
 					foreach (string Line in Receivers)
@@ -1315,6 +1316,7 @@ namespace UnrealAndroidFileTool
 			{
 				ArgIndex++;
 				List<string> Receivers = AndroidFileClient.GetInstalledReceivers(Device.Substring(1));
+				AndroidFileClient.GetInstalledActivities(Device.Substring(1)).ForEach(item => { if (!Receivers.Contains(item)) Receivers.Add(item); });
 				if (Receivers.Count > 0)
 				{
 					foreach (string Line in Receivers)
