Add -f option to UAFT for command

#android

diff --git a/Engine/Source/Programs/UnrealAndroidFileTool/Program.cs b/Engine/Source/Programs/UnrealAndroidFileTool/Program.cs
--- a/Engine/Source/Programs/UnrealAndroidFileTool/Program.cs
+++ b/Engine/Source/Programs/UnrealAndroidFileTool/Program.cs
@@ -12,12 +12,13 @@ using System.Linq;
 using AutomationTool;
 using EpicGames.Core;
 using Microsoft.Extensions.Logging;
+using Microsoft.CodeAnalysis.CSharp.Syntax;
 
 namespace UnrealAndroidFileTool
 {
 	class Program
 	{
-		private static string AppVersion = "1.0.5";
+		private static string AppVersion = "1.0.6";
 		private static ILogger Logger = Log.Logger;
 
 		static void ShowHelp(string command)
@@ -28,43 +29,43 @@ namespace UnrealAndroidFileTool
 					"Copyright Epic Games, Inc. All Rights Reserved.\n\n" +
 					"UnrealAndroidFileTool [-s Device] [-ip ipAddress] [-t port] [-p PackageName] [-k token] command\n\n" +
 					"Commands:\n\n" +
-					"devices                Show connected devices (@ prefix for authorized).\n" +
-					"packages               Show list of packages with receiver enabled.\n" +
-					"stop-all [-w]          Send stop request to all packages with receivers.\n" +
-					"help [command]         Get help about a command.\n" +
-					"shell                  Start interactive mode.\n" +
-					"quit                   Exit interactive mode.\n" +
-					"exit                   Exit interactive mode.\n" +
-					"terminate              Stop remote server.\n" +
-					"query [key]            Show all variables or a specific key.\n" +
-					"getprop [key]          Get property for key.\n" + 
-					"cd path                Set base directory (value of ^^).\n" +
-					"pwd                    Show current base directory.\n" +
-					"direxists path         Returns true if directory path exists.\n" +
-//					"dirlist [RSA:]path     Show contents of directory path.\n" +
-					"ls [-lsRf] path        Show contents of directory path.\n" +
-//					"dirlistflat path       Show recursive contents of directories grouped.\n" +
-//					"dircreate path         Create directory path (recursive).\n" +
-					"mkdir path             Create directory path (recursive).\n" +
-//					"dirdelete path         Delete directory path (recursive).\n" +
-					"rmdir path             Delete directory path (recursive).\n" +
-					"fileexists file        Returns true if file exists.\n" +
-//					"filedelete file        Delete file.\n" +
-					"rm file                Delete file.\n" +
-//					"filecopy src dst       Copies file from src to dst on device.\n" +
-					"cp src dst             Copies file from src to dst on device.\n" +
-//					"filemove src dst       Moves file from src to dst on device.\n" +
-					"mv src dst             Moves file from src to dst on device.\n" +
-//					"fileread src dst       Pulls src file from device to local dst.\n" +
-					"pull src dst           Pulls src file from device to local dst.\n" +
-					"pulldir src dst        Pulls src directory from device to local dst.\n" +
-//					"filewrite src dst      Pushes local src file to device dst.\n" +
-					"push [-c] src dst      Pushes local src file or directories to device dst.\n" +
-					"command [data]         Writes data to commandline file, or shows contents.\n" +
-					"addcommand [data]      Adds data to commandline file.\n" +
-					"delcommand [data]      Removes data from commandline file.\n" +
-					"cat file               Writes file contents from device to output.\n" +
-					"deploy [-c] file       Reads a text file with deployment source/dest pairs.\n" +
+					"devices                  Show connected devices (@ prefix for authorized).\n" +
+					"packages                 Show list of packages with receiver enabled.\n" +
+					"stop-all [-w]            Send stop request to all packages with receivers.\n" +
+					"help [command]           Get help about a command.\n" +
+					"shell                    Start interactive mode.\n" +
+					"quit                     Exit interactive mode.\n" +
+					"exit                     Exit interactive mode.\n" +
+					"terminate                Stop remote server.\n" +
+					"query [key]              Show all variables or a specific key.\n" +
+					"getprop [key]            Get property for key.\n" + 
+					"cd path                  Set base directory (value of ^^).\n" +
+					"pwd                      Show current base directory.\n" +
+					"direxists path           Returns true if directory path exists.\n" +
+//					"dirlist [RSA:]path       Show contents of directory path.\n" +
+					"ls [-lsRf] path          Show contents of directory path.\n" +
+//					"dirlistflat path         Show recursive contents of directories grouped.\n" +
+//					"dircreate path           Create directory path (recursive).\n" +
+					"mkdir path               Create directory path (recursive).\n" +
+//					"dirdelete path           Delete directory path (recursive).\n" +
+					"rmdir path               Delete directory path (recursive).\n" +
+					"fileexists file          Returns true if file exists.\n" +
+//					"filedelete file          Delete file.\n" +
+					"rm file                  Delete file.\n" +
+//					"filecopy src dst         Copies file from src to dst on device.\n" +
+					"cp src dst               Copies file from src to dst on device.\n" +
+//					"filemove src dst         Moves file from src to dst on device.\n" +
+					"mv src dst               Moves file from src to dst on device.\n" +
+//					"fileread src dst         Pulls src file from device to local dst.\n" +
+					"pull src dst             Pulls src file from device to local dst.\n" +
+					"pulldir src dst          Pulls src directory from device to local dst.\n" +
+//					"filewrite src dst        Pushes local src file to device dst.\n" +
+					"push [-c] src dst        Pushes local src file or directories to device dst.\n" +
+					"command [-f file] [data] Writes [file] or [data] to commandline file, or shows contents.\n" +
+					"addcommand [data]        Adds data to commandline file.\n" +
+					"delcommand [data]        Removes data from commandline file.\n" +
+					"cat file                 Writes file contents from device to output.\n" +
+					"deploy [-c] file         Reads a text file with deployment source/dest pairs.\n" +
 					"\n" +
 					"Note: keys may be used at start of a path as shortcut.\n" +
 					"   Ex: fileread ^logfile game.log   - pulls logfile.\n" +
@@ -72,6 +73,13 @@ namespace UnrealAndroidFileTool
 					, AppVersion);
 				return;
 			}
+			if (command == "command")
+			{
+				Logger.LogInformation("No Args\t\tShow contents of commandline file");
+				Logger.LogInformation("-f [file]\tWrites contents of [file] to commandline file");
+				Logger.LogInformation("[data]\t\tOverwrites commandline file with [data]");
+				return;
+			}
 			if (command == "ls")
 			{
 				Logger.LogInformation("-l\tlist permissions");
@@ -805,9 +813,33 @@ namespace UnrealAndroidFileTool
 				if (ArgIndex < args.Length)
 				{
 					string command = args[ArgIndex++];
-					while (ArgIndex < args.Length)
+					if (command == "-f")
 					{
-						command = command + " " + args[ArgIndex++];
+						if (ArgIndex < args.Length && File.Exists(args[ArgIndex]))
+						{
+							string[] CommandLines = File.ReadAllLines(args[ArgIndex++]);
+							command = "";
+							if (CommandLines.Length > 0)
+							{
+								command = CommandLines[0];
+							}
+						}
+						else
+						{
+							Logger.LogError("missing filename");
+							if (!bInShell)
+							{
+								client.CloseConnection();
+								System.Environment.Exit(1);
+							}
+						}
+					}
+					else
+					{
+						while (ArgIndex < args.Length)
+						{
+							command = command + " " + args[ArgIndex++];
+						}
 					}
 					bool result = client.FileWriteString(command, "^commandfile");
 					Logger.LogInformation("{0}", (result ? "true" : "false"));
